#ue4
- Added UK2Node_LatentOnlineCall which automatically creates entries for all static factory methods in UOnlineBlueprintCallProxyBase subclasses

diff --git a/Engine/Source/Editor/OnlineBlueprintSupport/Classes/K2Node_LatentOnlineCall.h b/Engine/Source/Editor/OnlineBlueprintSupport/Classes/K2Node_LatentOnlineCall.h
new file mode 100644
--- /dev/null
+++ b/Engine/Source/Editor/OnlineBlueprintSupport/Classes/K2Node_LatentOnlineCall.h
@@ -0,0 +1,16 @@
+// Copyright 1998-2014 Epic Games, Inc. All Rights Reserved.
+
+#pragma once
+
+#include "K2Node_LatentOnlineCall.generated.h"
+
+// This node is a latent online subsystem call (handles scanning all UOnlineBlueprintCallProxyBase classes for static factory calls)
+UCLASS()
+class ONLINEBLUEPRINTSUPPORT_API UK2Node_LatentOnlineCall : public UK2Node_BaseAsyncTask
+{
+	GENERATED_UCLASS_BODY()
+
+	// UEdGraphNode interface
+	virtual void GetMenuEntries(FGraphContextMenuBuilder& ContextMenuBuilder) const OVERRIDE;
+	// End of UEdGraphNode interface
+};
diff --git a/Engine/Source/Editor/OnlineBlueprintSupport/Private/K2Node_LatentOnlineCall.cpp b/Engine/Source/Editor/OnlineBlueprintSupport/Private/K2Node_LatentOnlineCall.cpp
new file mode 100644
--- /dev/null
+++ b/Engine/Source/Editor/OnlineBlueprintSupport/Private/K2Node_LatentOnlineCall.cpp
@@ -0,0 +1,52 @@
+// Copyright 1998-2014 Epic Games, Inc. All Rights Reserved.
+
+#include "OnlineBlueprintSupportPrivatePCH.h"
+#include "OnlineBlueprintCallProxyBase.h"
+#include "K2Node_LatentOnlineCall.h"
+
+#define LOCTEXT_NAMESPACE "K2Node"
+
+UK2Node_LatentOnlineCall::UK2Node_LatentOnlineCall(const FPostConstructInitializeProperties& PCIP)
+	: Super(PCIP)
+{
+	ProxyActivateFunctionName = GET_FUNCTION_NAME_CHECKED(UOnlineBlueprintCallProxyBase, Activate);
+}
+
+void UK2Node_LatentOnlineCall::GetMenuEntries(FGraphContextMenuBuilder& ContextMenuBuilder) const
+{
+	const UEdGraphSchema_K2* K2Schema = GetDefault<UEdGraphSchema_K2>();
+	const bool bAllowLatentFuncs = (K2Schema->GetGraphType(ContextMenuBuilder.CurrentGraph) == GT_Ubergraph);
+
+
+	if (bAllowLatentFuncs)
+	{
+		for (TObjectIterator<UClass> ClassIt; ClassIt; ++ClassIt)
+		{
+			UClass* TestClass = *ClassIt;
+			if (TestClass->IsChildOf(UOnlineBlueprintCallProxyBase::StaticClass()) && !TestClass->HasAnyClassFlags(CLASS_Abstract | CLASS_Deprecated))
+			{
+				for (TFieldIterator<UFunction> FuncIt(TestClass, EFieldIteratorFlags::ExcludeSuper); FuncIt; ++FuncIt)
+				{
+					// See if the function is a static factory method for online proxies
+					UFunction* CandidateFunction = *FuncIt;
+
+					UObjectProperty* ReturnProperty = Cast<UObjectProperty>(CandidateFunction->GetReturnProperty());
+					const bool bValidReturnType = (ReturnProperty != nullptr) && (ReturnProperty->PropertyClass != nullptr) && (ReturnProperty->PropertyClass->IsChildOf(UOnlineBlueprintCallProxyBase::StaticClass()));
+
+					if (CandidateFunction->HasAllFunctionFlags(FUNC_Static) && bValidReturnType)
+					{
+						// Create a node template for this factory method
+						UK2Node_LatentOnlineCall* NodeTemplate = NewObject<UK2Node_LatentOnlineCall>(ContextMenuBuilder.OwnerOfTemporaries);
+						NodeTemplate->ProxyFactoryFunctionName = CandidateFunction->GetFName();
+						NodeTemplate->ProxyFactoryClass = TestClass;
+						NodeTemplate->ProxyClass = ReturnProperty->PropertyClass;
+
+						CreateDefaultMenuEntry(NodeTemplate, ContextMenuBuilder);
+					}
+				}
+			}
+		}
+	}
+}
+
+#undef LOCTEXT_NAMESPACE
diff --git a/Engine/Source/Runtime/Online/OnlineSubsystemUtils/Classes/OnlineBlueprintCallProxyBase.h b/Engine/Source/Runtime/Online/OnlineSubsystemUtils/Classes/OnlineBlueprintCallProxyBase.h
new file mode 100644
--- /dev/null
+++ b/Engine/Source/Runtime/Online/OnlineSubsystemUtils/Classes/OnlineBlueprintCallProxyBase.h
@@ -0,0 +1,16 @@
+// Copyright 1998-2014 Epic Games, Inc. All Rights Reserved.
+#pragma once
+
+#include "OnlineBlueprintCallProxyBase.generated.h"
+
+DECLARE_DYNAMIC_MULTICAST_DELEGATE(FEmptyOnlineDelegate);
+
+UCLASS(MinimalAPI)
+class UOnlineBlueprintCallProxyBase : public UObject
+{
+	GENERATED_UCLASS_BODY()
+
+	// Called to trigger the actual online action once the delegates have been bound
+	UFUNCTION()
+	virtual void Activate();
+};
diff --git a/Engine/Source/Runtime/Online/OnlineSubsystemUtils/Private/OnlineBlueprintCallProxyBase.cpp b/Engine/Source/Runtime/Online/OnlineSubsystemUtils/Private/OnlineBlueprintCallProxyBase.cpp
new file mode 100644
--- /dev/null
+++ b/Engine/Source/Runtime/Online/OnlineSubsystemUtils/Private/OnlineBlueprintCallProxyBase.cpp
@@ -0,0 +1,16 @@
+// Copyright 1998-2014 Epic Games, Inc. All Rights Reserved.
+
+#include "OnlineSubsystemUtilsPrivatePCH.h"
+#include "OnlineBlueprintCallProxyBase.h"
+
+//////////////////////////////////////////////////////////////////////////
+// UOnlineBlueprintCallProxyBase
+
+UOnlineBlueprintCallProxyBase::UOnlineBlueprintCallProxyBase(const FPostConstructInitializeProperties& PCIP)
+	: Super(PCIP)
+{
+}
+
+void UOnlineBlueprintCallProxyBase::Activate()
+{
+}
